Adding Chinese translation to it  #jira ueloc-4711 #fyi jamie.dale, michelle.zhao, alexis.zhang, yali.he

diff --git a/Templates/TP_HololensBP/Config/TemplateDefs.ini b/Templates/TP_HololensBP/Config/TemplateDefs.ini
--- a/Templates/TP_HololensBP/Config/TemplateDefs.ini
+++ b/Templates/TP_HololensBP/Config/TemplateDefs.ini
@@ -10,9 +10,13 @@ Categories=MFG
 
 ; The localized display names for this template. Each entry specifies a Language code and Unicode text to display in the New Project dialog.
 LocalizedDisplayNames=(Language="en", Text="Hololens Viewer")
+LocalizedDisplayNames=(Language="zh-Hans", Text="Hololens查看器")
 
 ; The localized descriptions for this template, which is displayed when the template is selected.
 LocalizedDescriptions=(Language="en", Text="This project template provides navigation and interaction for Hololens. Use as a starting point for collaborative industrial or architectural projects.\n\nThis template enables you to use Datasmith to import models from various design applications and easily light them.")
+LocalizedDescriptions=(Language="zh-Hans", Text="该项目模板提供面向Hololens的导航及互动功能，适合用于协作式工业或建筑项目。\n\n该模板允许使用Datasmith从多种设计软件中导入模型并轻松为其设置光照。")
+
+
 ; The sort key to use for the template to bypass the default alphabetical sorting.
 SortKey="3"
 
